Motion Design: implemented Data Link Request Task for Scene State. This allows a data link graph with bindable input data to be executed and its result return through an AnyStruct property reference.
That property reference, although allowing to bind to any struct, must be bound to the output struct of the data link graph.
Data Link Graph does not at the moment define an expected output struct. Further work on this could improve what could be bound.

diff --git a/Engine/Plugins/Experimental/SceneStateDataLink/SceneStateDataLink.uplugin b/Engine/Plugins/Experimental/SceneStateDataLink/SceneStateDataLink.uplugin
new file mode 100644
--- /dev/null
+++ b/Engine/Plugins/Experimental/SceneStateDataLink/SceneStateDataLink.uplugin
@@ -0,0 +1,38 @@
+{
+	"FileVersion": 3,
+	"Version": 1,
+	"VersionName": "1.0",
+	"FriendlyName": "Motion Design Scene State Data Link Bridge",
+	"Description": "Scene State Tasks that execute Data Link Graphs",
+	"Category": "Other",
+	"CreatedBy": "",
+	"CreatedByURL": "",
+	"DocsURL": "",
+	"MarketplaceURL": "",
+	"CanContainContent": true,
+	"IsBetaVersion": false,
+	"IsExperimentalVersion": true,
+	"Installed": false,
+	"Modules": [
+		{
+			"Name": "SceneStateDataLink",
+			"Type": "Runtime",
+			"LoadingPhase": "Default"
+		},
+		{
+			"Name": "SceneStateDataLinkEditor",
+			"Type": "Editor",
+			"LoadingPhase": "Default"
+		}
+	],
+	"Plugins": [
+		{
+			"Name": "SceneState",
+			"Enabled": true
+		},
+		{
+			"Name": "DataLink",
+			"Enabled": true
+		}
+	]
+}
diff --git a/Engine/Plugins/Experimental/SceneStateDataLink/Source/SceneStateDataLink/Private/SceneStateDataLinkModule.cpp b/Engine/Plugins/Experimental/SceneStateDataLink/Source/SceneStateDataLink/Private/SceneStateDataLinkModule.cpp
new file mode 100644
--- /dev/null
+++ b/Engine/Plugins/Experimental/SceneStateDataLink/Source/SceneStateDataLink/Private/SceneStateDataLinkModule.cpp
@@ -0,0 +1,5 @@
+// Copyright Epic Games, Inc. All Rights Reserved.
+
+#include "Modules/ModuleManager.h"
+
+IMPLEMENT_MODULE(FDefaultModuleImpl, SceneStateDataLink)
diff --git a/Engine/Plugins/Experimental/SceneStateDataLink/Source/SceneStateDataLink/Private/SceneStateDataLinkRequestTask.cpp b/Engine/Plugins/Experimental/SceneStateDataLink/Source/SceneStateDataLink/Private/SceneStateDataLinkRequestTask.cpp
new file mode 100644
--- /dev/null
+++ b/Engine/Plugins/Experimental/SceneStateDataLink/Source/SceneStateDataLink/Private/SceneStateDataLinkRequestTask.cpp
@@ -0,0 +1,89 @@
+// Copyright Epic Games, Inc. All Rights Reserved.
+
+#include "SceneStateDataLinkRequestTask.h"
+#include "DataLinkExecutor.h"
+#include "DataLinkExecutorArguments.h"
+#include "DataLinkGraph.h"
+#include "DataLinkUtils.h"
+#include "Misc/ScopeExit.h"
+#include "SceneStateExecutionContext.h"
+#include "Tasks/SceneStateTaskExecutionContext.h"
+
+FDataLinkInstance FSceneStateDataLinkRequestTaskInstance::CreateDataLinkInstance() const
+{
+	return FDataLinkInstance
+		{
+			.DataLinkGraph = DataLinkGraph,
+			.InputData = InputData
+		};
+}
+
+#if WITH_EDITOR
+const UScriptStruct* FSceneStateDataLinkRequestTask::OnGetTaskInstanceType() const
+{
+	return FInstanceDataType::StaticStruct();
+}
+
+void FSceneStateDataLinkRequestTask::OnBuildTaskInstance(UObject* InOuter, FStructView InTaskInstance) const
+{
+	FInstanceDataType& TaskInstance = InTaskInstance.Get<FInstanceDataType>();
+	UE::DataLink::SetInputData(TaskInstance.DataLinkGraph, TaskInstance.InputData);
+}
+#endif
+
+void FSceneStateDataLinkRequestTask::OnStart(const FSceneStateExecutionContext& InContext, FStructView InTaskInstance) const
+{
+	FInstanceDataType& Instance = InTaskInstance.Get<FInstanceDataType>();
+
+	Instance.Executor = FDataLinkExecutor::Create(FDataLinkExecutorArguments(Instance.CreateDataLinkInstance())
+#if WITH_DATALINK_CONTEXT
+		.SetContextName(InContext.GetExecutionContextName())
+#endif
+		.SetContextObject(InContext.GetContextObject())
+		.SetOnFinish(FOnDataLinkExecutionFinished::CreateStatic(&FSceneStateDataLinkRequestTask::OnExecutionFinished, UE::SceneState::FTaskExecutionContext(*this, InContext))));
+
+	Instance.Executor->Run();
+}
+
+void FSceneStateDataLinkRequestTask::OnExecutionFinished(const FDataLinkExecutor& InExecutor, EDataLinkExecutionResult InResult, FConstStructView InOutputDataView, UE::SceneState::FTaskExecutionContext InTaskContext)
+{
+	// Ensure that any exit path finishes the task
+	ON_SCOPE_EXIT
+	{
+		InTaskContext.FinishTask();
+	};
+
+	const FSceneStateExecutionContext* const ExecutionContext = InTaskContext.GetExecutionContext();
+	if (!ExecutionContext)
+	{
+		return;
+	}
+
+	FInstanceDataType* const Instance = InTaskContext.GetTaskInstance().GetPtr<FInstanceDataType>();
+
+	// Ensure the instance gotten from this task context and executor remained the same
+	if (!Instance || !ensure(Instance->Executor.Get() == &InExecutor))
+	{
+		return;
+	}
+
+	Instance->Executor.Reset();
+
+	UE::SceneState::FResolvePropertyResult Result;
+	if (!UE::SceneState::ResolveProperty(*ExecutionContext, Instance->OutputTarget, Result))
+	{
+		return;
+	}
+
+	// The validity of these objects are guaranteed if ResolveProperty returns true
+	check(Result.ResolvedReference && Result.ResolvedReference->SourceLeafProperty);
+
+	// Validate that the reference struct and output struct match
+	const FStructProperty* StructProperty = CastField<FStructProperty>(Result.ResolvedReference->SourceLeafProperty);
+	if (!StructProperty || !StructProperty->Struct || StructProperty->Struct != InOutputDataView.GetScriptStruct())
+	{
+		return;
+	}
+
+	StructProperty->Struct->CopyScriptStruct(Result.ValuePtr, InOutputDataView.GetMemory());
+}
diff --git a/Engine/Plugins/Experimental/SceneStateDataLink/Source/SceneStateDataLink/Public/SceneStateDataLinkRequestTask.h b/Engine/Plugins/Experimental/SceneStateDataLink/Source/SceneStateDataLink/Public/SceneStateDataLinkRequestTask.h
new file mode 100644
--- /dev/null
+++ b/Engine/Plugins/Experimental/SceneStateDataLink/Source/SceneStateDataLink/Public/SceneStateDataLinkRequestTask.h
@@ -0,0 +1,69 @@
+// Copyright Epic Games, Inc. All Rights Reserved.
+
+#pragma once
+
+#include "DataLinkInstance.h"
+#include "SceneStatePropertyReference.h"
+#include "Tasks/SceneStateTask.h"
+#include "Tasks/SceneStateTaskInstance.h"
+#include "SceneStateDataLinkRequestTask.generated.h"
+
+class FDataLinkExecutor;
+class UDataLinkGraph;
+enum class EDataLinkExecutionResult : uint8;
+
+namespace UE::SceneState
+{
+	struct FTaskExecutionContext;
+}
+
+USTRUCT()
+struct FSceneStateDataLinkRequestTaskInstance : public FSceneStateTaskInstance
+{
+	GENERATED_BODY()
+
+	/** Creates a FDataLinkInstance from this Task Instance */
+	FDataLinkInstance CreateDataLinkInstance() const;
+
+	/** The data link graph to execute */
+	UPROPERTY(EditAnywhere, Category="Data Link", meta=(NoBinding))
+	TObjectPtr<UDataLinkGraph> DataLinkGraph;
+
+	/**
+	 * Input data required to run the Data Link Graph.
+	 * The type of input data is determined by the Task's Data Link Graph property
+	 */
+	UPROPERTY(EditAnywhere, EditFixedSize, Category="Data Link", meta=(EditFixedOrder, NoBindingSelfOnly))
+	TArray<FDataLinkInputData> InputData;
+
+	/**
+	 * The output struct reference to write the data link results to.
+	 * NOTE: Must match the data link output struct type, or will fail to write!
+	 */
+	UPROPERTY(EditAnywhere, Category="Data Link", meta=(RefType="AnyStruct"))
+	FSceneStatePropertyReference OutputTarget;
+
+	/** The object handling the execution */
+	TSharedPtr<FDataLinkExecutor> Executor;
+};
+
+USTRUCT(DisplayName="Data Link Request", Category="Data Link", meta=(ToolTip="Requests data via a Data Link Graph"))
+struct FSceneStateDataLinkRequestTask : public FSceneStateTask
+{
+	GENERATED_BODY()
+
+	using FInstanceDataType = FSceneStateDataLinkRequestTaskInstance;
+
+protected:
+	//~ Begin FSceneStateTask
+#if WITH_EDITOR
+	virtual const UScriptStruct* OnGetTaskInstanceType() const override;
+	virtual void OnBuildTaskInstance(UObject* InOuter, FStructView InTaskInstance) const override;
+#endif
+	virtual void OnStart(const FSceneStateExecutionContext& InContext, FStructView InTaskInstance) const override;
+	//~ End FSceneStateTask
+
+private:
+	/** Called when a data link request has completed */
+	static void OnExecutionFinished(const FDataLinkExecutor& InExecutor, EDataLinkExecutionResult InResult, FConstStructView InOutputDataView, UE::SceneState::FTaskExecutionContext InTaskContext);
+};
diff --git a/Engine/Plugins/Experimental/SceneStateDataLink/Source/SceneStateDataLink/SceneStateDataLink.Build.cs b/Engine/Plugins/Experimental/SceneStateDataLink/Source/SceneStateDataLink/SceneStateDataLink.Build.cs
new file mode 100644
--- /dev/null
+++ b/Engine/Plugins/Experimental/SceneStateDataLink/Source/SceneStateDataLink/SceneStateDataLink.Build.cs
@@ -0,0 +1,23 @@
+// Copyright Epic Games, Inc. All Rights Reserved.
+
+using UnrealBuildTool;
+
+public class SceneStateDataLink : ModuleRules
+{
+	public SceneStateDataLink(ReadOnlyTargetRules Target) : base(Target)
+	{
+		PCHUsage = ModuleRules.PCHUsageMode.UseExplicitOrSharedPCHs;
+
+		PublicDependencyModuleNames.AddRange(
+			new string[]
+			{
+				"Core",
+				"CoreUObject",
+				"DataLink",
+				"Engine",
+				"SceneState",
+				"SceneStateBinding",
+			}
+		);
+	}
+}
diff --git a/Engine/Plugins/Experimental/SceneStateDataLink/Source/SceneStateDataLinkEditor/Private/DetailsView/SceneStateDataLinkRequestTaskDetails.cpp b/Engine/Plugins/Experimental/SceneStateDataLink/Source/SceneStateDataLinkEditor/Private/DetailsView/SceneStateDataLinkRequestTaskDetails.cpp
new file mode 100644
--- /dev/null
+++ b/Engine/Plugins/Experimental/SceneStateDataLink/Source/SceneStateDataLinkEditor/Private/DetailsView/SceneStateDataLinkRequestTaskDetails.cpp
@@ -0,0 +1,95 @@
+// Copyright Epic Games, Inc. All Rights Reserved.
+
+#include "SceneStateDataLinkRequestTaskDetails.h"
+#include "DataLinkGraph.h"
+#include "DataLinkUtils.h"
+#include "IDetailChildrenBuilder.h"
+#include "PropertyHandle.h"
+#include "SceneStateBlueprintEditorUtils.h"
+#include "SceneStateBlueprintEditorUtils.h"
+#include "SceneStateDataLinkRequestTask.h"
+
+namespace UE::SceneStateDataLink
+{
+
+void FRequestTaskInstanceDetails::CustomizeHeader(TSharedRef<IPropertyHandle> InPropertyHandle, FDetailWidgetRow& InHeaderRow, IPropertyTypeCustomizationUtils& InCustomizationUtils)
+{
+}
+
+FRequestTaskInstanceDetails::~FRequestTaskInstanceDetails()
+{
+	UDataLinkGraph::OnGraphCompiled().RemoveAll(this);
+}
+
+void FRequestTaskInstanceDetails::CustomizeChildren(TSharedRef<IPropertyHandle> InPropertyHandle, IDetailChildrenBuilder& InChildBuilder, IPropertyTypeCustomizationUtils& InCustomizationUtils)
+{
+	UDataLinkGraph::OnGraphCompiled().RemoveAll(this);
+	UDataLinkGraph::OnGraphCompiled().AddSP(this, &FRequestTaskInstanceDetails::OnGraphCompiled);
+
+	DataLinkGraphHandle = InPropertyHandle->GetChildHandle(GET_MEMBER_NAME_CHECKED(FSceneStateDataLinkRequestTaskInstance, DataLinkGraph));
+	check(DataLinkGraphHandle.IsValid());
+	DataLinkGraphHandle->MarkHiddenByCustomization();
+	DataLinkGraphHandle->SetOnPropertyValueChanged(FSimpleDelegate::CreateSP(this, &FRequestTaskInstanceDetails::OnGraphChanged));
+
+	InputDataHandle = InPropertyHandle->GetChildHandle(GET_MEMBER_NAME_CHECKED(FDataLinkInstance, InputData));
+	check(InputDataHandle.IsValid());
+	InputDataHandle->MarkHiddenByCustomization();
+
+	const FGuid TaskId = SceneState::Editor::FindTaskId(InPropertyHandle);
+	UE::SceneState::Editor::AssignBindingId(InputDataHandle.ToSharedRef(), TaskId);
+
+	InChildBuilder.AddProperty(DataLinkGraphHandle.ToSharedRef());
+	InChildBuilder.AddProperty(InputDataHandle.ToSharedRef());
+
+	UpdateInputData();
+}
+
+UDataLinkGraph* FRequestTaskInstanceDetails::GetDataLinkGraph() const
+{
+	UObject* DataLinkGraph = nullptr;
+	if (DataLinkGraphHandle.IsValid() && DataLinkGraphHandle->GetValue(DataLinkGraph) == FPropertyAccess::Success)
+	{
+		return Cast<UDataLinkGraph>(DataLinkGraph);
+	}
+	return nullptr;
+}
+
+TArray<FDataLinkInputData>* FRequestTaskInstanceDetails::GetInputData() const
+{
+	void* InputData = nullptr;
+	if (InputDataHandle.IsValid() && InputDataHandle->GetValueData(InputData) == FPropertyAccess::Success)
+	{
+		return static_cast<TArray<FDataLinkInputData>*>(InputData);
+	}
+	return nullptr;
+}
+
+void FRequestTaskInstanceDetails::OnGraphCompiled(UDataLinkGraph* InDataLinkGraph)
+{
+	if (InDataLinkGraph == GetDataLinkGraph())
+	{
+		OnGraphChanged();
+	}
+}
+
+void FRequestTaskInstanceDetails::OnGraphChanged()
+{
+	if (!InputDataHandle.IsValid())
+	{
+		return;
+	}
+
+	InputDataHandle->NotifyPreChange();
+	UpdateInputData();
+	InputDataHandle->NotifyPostChange(EPropertyChangeType::ValueSet);
+}
+
+void FRequestTaskInstanceDetails::UpdateInputData()
+{
+	if (TArray<FDataLinkInputData>* InputData = GetInputData())
+	{
+		UE::DataLink::SetInputData(GetDataLinkGraph(), *InputData);
+	}
+}
+
+} // UE::SceneStateDataLink
diff --git a/Engine/Plugins/Experimental/SceneStateDataLink/Source/SceneStateDataLinkEditor/Private/DetailsView/SceneStateDataLinkRequestTaskDetails.h b/Engine/Plugins/Experimental/SceneStateDataLink/Source/SceneStateDataLinkEditor/Private/DetailsView/SceneStateDataLinkRequestTaskDetails.h
new file mode 100644
--- /dev/null
+++ b/Engine/Plugins/Experimental/SceneStateDataLink/Source/SceneStateDataLinkEditor/Private/DetailsView/SceneStateDataLinkRequestTaskDetails.h
@@ -0,0 +1,48 @@
+// Copyright Epic Games, Inc. All Rights Reserved.
+
+#pragma once
+
+#include "IPropertyTypeCustomization.h"
+
+class UDataLinkGraph;
+struct FDataLinkInputData;
+
+namespace UE::SceneStateDataLink
+{
+
+/** Details Customization for FSceneStateDataLinkRequestTaskInstance */
+class FRequestTaskInstanceDetails : public IPropertyTypeCustomization
+{
+public:
+	static TSharedRef<IPropertyTypeCustomization> MakeInstance()
+	{
+		return MakeShared<FRequestTaskInstanceDetails>();
+	}
+
+	virtual ~FRequestTaskInstanceDetails() override;
+
+protected:
+	//~ Begin IPropertyTypeCustomization
+	virtual void CustomizeHeader(TSharedRef<IPropertyHandle> InPropertyHandle, FDetailWidgetRow& InHeaderRow, IPropertyTypeCustomizationUtils& InCustomizationUtils) override;
+	virtual void CustomizeChildren(TSharedRef<IPropertyHandle> InPropertyHandle, IDetailChildrenBuilder& InChildBuilder, IPropertyTypeCustomizationUtils& InCustomizationUtils) override;
+	//~ End IPropertyTypeCustomization
+
+private:
+	/** Gets the value of the data link graph property handle */
+	UDataLinkGraph* GetDataLinkGraph() const;
+	/** Gets the value of the input data property handle */
+	TArray<FDataLinkInputData>* GetInputData() const;
+
+	/** Called when a data link graph has compiled. Used to update the input data */
+	void OnGraphCompiled(UDataLinkGraph* InDataLinkGraph);
+	/** Called when the graph has changed to update the input data */
+	void OnGraphChanged();
+
+	/** Updates the input data array to match the data link graph's input pins */
+	void UpdateInputData();
+
+	TSharedPtr<IPropertyHandle> DataLinkGraphHandle;
+	TSharedPtr<IPropertyHandle> InputDataHandle;
+};
+
+} // UE::SceneStateDataLink
diff --git a/Engine/Plugins/Experimental/SceneStateDataLink/Source/SceneStateDataLinkEditor/Private/SceneStateDataLinkEditorModule.cpp b/Engine/Plugins/Experimental/SceneStateDataLink/Source/SceneStateDataLinkEditor/Private/SceneStateDataLinkEditorModule.cpp
new file mode 100644
--- /dev/null
+++ b/Engine/Plugins/Experimental/SceneStateDataLink/Source/SceneStateDataLinkEditor/Private/SceneStateDataLinkEditorModule.cpp
@@ -0,0 +1,44 @@
+// Copyright Epic Games, Inc. All Rights Reserved.
+
+#include "SceneStateDataLinkEditorModule.h"
+#include "DetailsView/SceneStateDataLinkRequestTaskDetails.h"
+#include "Modules/ModuleManager.h"
+#include "PropertyEditorModule.h"
+#include "SceneStateDataLinkRequestTask.h"
+
+IMPLEMENT_MODULE(UE::SceneStateDataLink::FEditorModule, SceneStateDataLinkEditor)
+
+namespace UE::SceneStateDataLink
+{
+
+void FEditorModule::StartupModule()
+{
+	RegisterDetailCustomizations();
+}
+
+void FEditorModule::ShutdownModule()
+{
+	UnregisterDetailCustomizations();
+}
+
+void FEditorModule::RegisterDetailCustomizations()
+{
+	FPropertyEditorModule& PropertyEditorModule = FModuleManager::LoadModuleChecked<FPropertyEditorModule>("PropertyEditor");
+
+	PropertyEditorModule.RegisterCustomPropertyTypeLayout(CustomizedTypes.Add_GetRef(FSceneStateDataLinkRequestTaskInstance::StaticStruct()->GetFName())
+		, FOnGetPropertyTypeCustomizationInstance::CreateStatic(&FRequestTaskInstanceDetails::MakeInstance));
+}
+
+void FEditorModule::UnregisterDetailCustomizations()
+{
+	if (FPropertyEditorModule* PropertyEditorModule = FModuleManager::GetModulePtr<FPropertyEditorModule>("PropertyEditor"))
+	{
+		for (FName CustomizedType : CustomizedTypes)
+		{
+			PropertyEditorModule->UnregisterCustomPropertyTypeLayout(CustomizedType);
+		}
+		CustomizedTypes.Reset();
+	}
+}
+
+} // UE::SceneStateDataLink
diff --git a/Engine/Plugins/Experimental/SceneStateDataLink/Source/SceneStateDataLinkEditor/Private/SceneStateDataLinkEditorModule.h b/Engine/Plugins/Experimental/SceneStateDataLink/Source/SceneStateDataLinkEditor/Private/SceneStateDataLinkEditorModule.h
new file mode 100644
--- /dev/null
+++ b/Engine/Plugins/Experimental/SceneStateDataLink/Source/SceneStateDataLinkEditor/Private/SceneStateDataLinkEditorModule.h
@@ -0,0 +1,27 @@
+// Copyright Epic Games, Inc. All Rights Reserved.
+
+#pragma once
+
+#include "Containers/Array.h"
+#include "Modules/ModuleInterface.h"
+#include "UObject/NameTypes.h"
+
+namespace UE::SceneStateDataLink
+{
+
+class FEditorModule : public IModuleInterface
+{
+public:
+	//~ Begin IModuleInterface
+	virtual void StartupModule() override;
+	virtual void ShutdownModule() override;
+	//~ End IModuleInterface
+
+private:
+	void RegisterDetailCustomizations();
+	void UnregisterDetailCustomizations();
+
+	TArray<FName> CustomizedTypes;
+};
+
+} // UE::SceneStateDataLink
\ No newline at end of file
diff --git a/Engine/Plugins/Experimental/SceneStateDataLink/Source/SceneStateDataLinkEditor/SceneStateDataLinkEditor.Build.cs b/Engine/Plugins/Experimental/SceneStateDataLink/Source/SceneStateDataLinkEditor/SceneStateDataLinkEditor.Build.cs
new file mode 100644
--- /dev/null
+++ b/Engine/Plugins/Experimental/SceneStateDataLink/Source/SceneStateDataLinkEditor/SceneStateDataLinkEditor.Build.cs
@@ -0,0 +1,24 @@
+// Copyright Epic Games, Inc. All Rights Reserved.
+
+using UnrealBuildTool;
+
+public class SceneStateDataLinkEditor : ModuleRules
+{
+    public SceneStateDataLinkEditor(ReadOnlyTargetRules Target) : base(Target)
+    {
+        PCHUsage = ModuleRules.PCHUsageMode.UseExplicitOrSharedPCHs;
+
+        PrivateDependencyModuleNames.AddRange(
+            new string[]
+            {
+                "Core",
+                "CoreUObject",
+                "DataLink",
+                "Engine",
+                "PropertyEditor",
+                "SceneStateBlueprintEditor",
+                "SceneStateDataLink",
+            }
+        );
+    }
+}
