Motion Design: fixed scene state data link task instance details customization not including the output data property

diff --git a/Engine/Plugins/Experimental/SceneStateDataLink/Source/SceneStateDataLinkEditor/Private/DetailsView/SceneStateDataLinkRequestTaskDetails.cpp b/Engine/Plugins/Experimental/SceneStateDataLink/Source/SceneStateDataLinkEditor/Private/DetailsView/SceneStateDataLinkRequestTaskDetails.cpp
--- a/Engine/Plugins/Experimental/SceneStateDataLink/Source/SceneStateDataLinkEditor/Private/DetailsView/SceneStateDataLinkRequestTaskDetails.cpp
+++ b/Engine/Plugins/Experimental/SceneStateDataLink/Source/SceneStateDataLinkEditor/Private/DetailsView/SceneStateDataLinkRequestTaskDetails.cpp
@@ -31,15 +31,20 @@ void FRequestTaskInstanceDetails::CustomizeChildren(TSharedRef<IPropertyHandle>
 	DataLinkGraphHandle->MarkHiddenByCustomization();
 	DataLinkGraphHandle->SetOnPropertyValueChanged(FSimpleDelegate::CreateSP(this, &FRequestTaskInstanceDetails::OnGraphChanged));
 
-	InputDataHandle = InPropertyHandle->GetChildHandle(GET_MEMBER_NAME_CHECKED(FDataLinkInstance, InputData));
+	InputDataHandle = InPropertyHandle->GetChildHandle(GET_MEMBER_NAME_CHECKED(FSceneStateDataLinkRequestTaskInstance, InputData));
 	check(InputDataHandle.IsValid());
 	InputDataHandle->MarkHiddenByCustomization();
 
+	TSharedRef<IPropertyHandle> OutputDataHandle = InPropertyHandle->GetChildHandle(GET_MEMBER_NAME_CHECKED(FSceneStateDataLinkRequestTaskInstance, OutputTarget)).ToSharedRef();
+	OutputDataHandle->MarkHiddenByCustomization();
+
 	const FGuid TaskId = SceneState::Editor::FindTaskId(InPropertyHandle);
 	UE::SceneState::Editor::AssignBindingId(InputDataHandle.ToSharedRef(), TaskId);
+	UE::SceneState::Editor::AssignBindingId(OutputDataHandle, TaskId);
 
 	InChildBuilder.AddProperty(DataLinkGraphHandle.ToSharedRef());
 	InChildBuilder.AddProperty(InputDataHandle.ToSharedRef());
+	InChildBuilder.AddProperty(OutputDataHandle);
 
 	UpdateInputData();
 }
