Fixed static analyzer warning

diff --git a/Engine/Plugins/Experimental/SceneStateDataLink/Source/SceneStateDataLink/Private/SceneStateDataLinkRequestTask.cpp b/Engine/Plugins/Experimental/SceneStateDataLink/Source/SceneStateDataLink/Private/SceneStateDataLinkRequestTask.cpp
--- a/Engine/Plugins/Experimental/SceneStateDataLink/Source/SceneStateDataLink/Private/SceneStateDataLinkRequestTask.cpp
+++ b/Engine/Plugins/Experimental/SceneStateDataLink/Source/SceneStateDataLink/Private/SceneStateDataLinkRequestTask.cpp
@@ -82,7 +82,14 @@ void FSceneStateDataLinkRequestTask::OnExecutionFinished(const FDataLinkExecutor
 
 	// Validate that the reference struct and output struct match
 	const FStructProperty* StructProperty = CastField<FStructProperty>(Result.ResolvedReference->SourceLeafProperty);
-	if (!StructProperty || !StructProperty->Struct || StructProperty->Struct != InOutputDataView.GetScriptStruct())
+	if (!StructProperty)
+	{
+		UE_LOG(LogSceneStateDataLink, Error, TEXT("[%s] Data Link Request Task failed. Output Target is not bound to a struct property")
+			, *ExecutionContext->GetExecutionContextName());
+		return;
+	}
+
+	if (!StructProperty->Struct || StructProperty->Struct != InOutputDataView.GetScriptStruct())
 	{
 		UE_LOG(LogSceneStateDataLink, Warning, TEXT("[%s] Data Link Request Task failed. Output Target struct ('%s') and Data Link Output struct ('%s') do not match!")
 			, *ExecutionContext->GetExecutionContextName()
