Motion Design Scene State: Renamed Data Link request task to Run data link task

diff --git a/Engine/Plugins/Experimental/SceneStateDataLink/Source/SceneStateDataLink/Private/SceneStateDataLinkRequestTask.cpp b/Engine/Plugins/Experimental/SceneStateDataLink/Source/SceneStateDataLink/Private/SceneStateDataLinkRequestTask.cpp
deleted file mode 100644
diff --git a/Engine/Plugins/Experimental/SceneStateDataLink/Source/SceneStateDataLink/Private/SceneStateRunDataLinkTask.cpp b/Engine/Plugins/Experimental/SceneStateDataLink/Source/SceneStateDataLink/Private/SceneStateRunDataLinkTask.cpp
new file mode 100644
--- /dev/null
+++ b/Engine/Plugins/Experimental/SceneStateDataLink/Source/SceneStateDataLink/Private/SceneStateRunDataLinkTask.cpp
@@ -0,0 +1,124 @@
+// Copyright Epic Games, Inc. All Rights Reserved.
+
+#include "SceneStateRunDataLinkTask.h"
+#include "DataLinkExecutor.h"
+#include "DataLinkExecutorArguments.h"
+#include "DataLinkGraph.h"
+#include "DataLinkUtils.h"
+#include "Misc/ScopeExit.h"
+#include "SceneStateDataLinkLog.h"
+#include "SceneStateExecutionContext.h"
+#include "Tasks/SceneStateTaskExecutionContext.h"
+
+FDataLinkInstance FSceneStateDataLinkRequestTaskInstance::CreateDataLinkInstance() const
+{
+	return FDataLinkInstance
+		{
+			.DataLinkGraph = DataLinkGraph,
+			.InputData = InputData
+		};
+}
+
+#if WITH_EDITOR
+const UScriptStruct* FSceneStateRunDataLinkTask::OnGetTaskInstanceType() const
+{
+	return FInstanceDataType::StaticStruct();
+}
+
+void FSceneStateRunDataLinkTask::OnBuildTaskInstance(UObject* InOuter, FStructView InTaskInstance) const
+{
+	FInstanceDataType& TaskInstance = InTaskInstance.Get<FInstanceDataType>();
+	UE::DataLink::SetInputData(TaskInstance.DataLinkGraph, TaskInstance.InputData);
+}
+#endif
+
+void FSceneStateRunDataLinkTask::OnStart(const FSceneStateExecutionContext& InContext, FStructView InTaskInstance) const
+{
+	FInstanceDataType& Instance = InTaskInstance.Get<FInstanceDataType>();
+
+	const UE::SceneState::FTaskExecutionContext TaskExecutionContext(*this, InContext);
+
+	Instance.Executor = FDataLinkExecutor::Create(FDataLinkExecutorArguments(Instance.CreateDataLinkInstance())
+#if WITH_DATALINK_CONTEXT
+		.SetContextName(InContext.GetExecutionContextName())
+#endif
+		.SetContextObject(InContext.GetContextObject())
+		.SetOnOutputData(FOnDataLinkOutputData::CreateStatic(&FSceneStateRunDataLinkTask::OnOutputData, TaskExecutionContext))
+		.SetOnFinished(FOnDataLinkExecutionFinished::CreateStatic(&FSceneStateRunDataLinkTask::OnFinished, TaskExecutionContext)));
+
+	Instance.Executor->Run();
+}
+
+void FSceneStateRunDataLinkTask::OnStop(const FSceneStateExecutionContext& InContext, FStructView InTaskInstance, ESceneStateTaskStopReason InStopReason) const
+{
+	FInstanceDataType& Instance = InTaskInstance.Get<FInstanceDataType>();
+	if (Instance.Executor.IsValid())
+	{
+		Instance.Executor->Stop();
+		Instance.Executor.Reset();
+	}
+}
+
+void FSceneStateRunDataLinkTask::OnOutputData(const FDataLinkExecutor& InExecutor, FConstStructView InOutputDataView, UE::SceneState::FTaskExecutionContext InTaskContext)
+{
+	if (!InOutputDataView.IsValid())
+	{
+		InTaskContext.FinishTask();
+		return;
+	}
+
+	const FSceneStateExecutionContext* const ExecutionContext = InTaskContext.GetExecutionContext();
+	if (!ExecutionContext)
+	{
+		InTaskContext.FinishTask();
+		return;
+	}
+
+	FInstanceDataType* const Instance = InTaskContext.GetTaskInstance().GetPtr<FInstanceDataType>();
+
+	// Ensure the instance gotten from this task context and executor remained the same
+	if (!Instance || !ensure(Instance->Executor.Get() == &InExecutor))
+	{
+		return;
+	}
+
+	Instance->Executor.Reset();
+
+	UE::SceneState::FResolvePropertyResult Result;
+	if (!UE::SceneState::ResolveProperty(*ExecutionContext, Instance->OutputTarget, Result))
+	{
+		UE_LOG(LogSceneStateDataLink, Warning, TEXT("[%s] Data Link Request Task failed. Output Target could not be resolved."), *ExecutionContext->GetExecutionContextName());
+		InTaskContext.FinishTask();
+		return;
+	}
+
+	// The validity of these objects are guaranteed if ResolveProperty returns true
+	check(Result.ResolvedReference && Result.ResolvedReference->SourceLeafProperty);
+
+	// Validate that the reference struct and output struct match
+	const FStructProperty* StructProperty = CastField<FStructProperty>(Result.ResolvedReference->SourceLeafProperty);
+	if (!StructProperty)
+	{
+		UE_LOG(LogSceneStateDataLink, Error, TEXT("[%s] Data Link Request Task failed. Output Target is not bound to a struct property")
+			, *ExecutionContext->GetExecutionContextName());
+		InTaskContext.FinishTask();
+		return;
+	}
+
+	if (!StructProperty->Struct || StructProperty->Struct != InOutputDataView.GetScriptStruct())
+	{
+		UE_LOG(LogSceneStateDataLink, Warning, TEXT("[%s] Data Link Request Task failed. Output Target struct ('%s') and Data Link Output struct ('%s') do not match!")
+			, *ExecutionContext->GetExecutionContextName()
+			, *GetNameSafe(StructProperty->Struct)
+			, *GetNameSafe(InOutputDataView.GetScriptStruct()));
+		InTaskContext.FinishTask();
+		return;
+	}
+
+	StructProperty->Struct->CopyScriptStruct(Result.ValuePtr, InOutputDataView.GetMemory());
+}
+
+void FSceneStateRunDataLinkTask::OnFinished(const FDataLinkExecutor& InExecutor, EDataLinkExecutionResult InExecutionResult, UE::SceneState::FTaskExecutionContext InTaskContext)
+{
+	InTaskContext.FinishTask();
+}
diff --git a/Engine/Plugins/Experimental/SceneStateDataLink/Source/SceneStateDataLink/Public/SceneStateDataLinkRequestTask.h b/Engine/Plugins/Experimental/SceneStateDataLink/Source/SceneStateDataLink/Public/SceneStateDataLinkRequestTask.h
deleted file mode 100644
diff --git a/Engine/Plugins/Experimental/SceneStateDataLink/Source/SceneStateDataLink/Public/SceneStateRunDataLinkTask.h b/Engine/Plugins/Experimental/SceneStateDataLink/Source/SceneStateDataLink/Public/SceneStateRunDataLinkTask.h
new file mode 100644
--- /dev/null
+++ b/Engine/Plugins/Experimental/SceneStateDataLink/Source/SceneStateDataLink/Public/SceneStateRunDataLinkTask.h
@@ -0,0 +1,74 @@
+// Copyright Epic Games, Inc. All Rights Reserved.
+
+#pragma once
+
+#include "DataLinkInstance.h"
+#include "SceneStatePropertyReference.h"
+#include "Tasks/SceneStateTask.h"
+#include "Tasks/SceneStateTaskInstance.h"
+#include "SceneStateRunDataLinkTask.generated.h"
+
+class FDataLinkExecutor;
+class UDataLinkGraph;
+enum class EDataLinkExecutionResult : uint8;
+
+namespace UE::SceneState
+{
+	struct FTaskExecutionContext;
+}
+
+USTRUCT()
+struct FSceneStateDataLinkRequestTaskInstance : public FSceneStateTaskInstance
+{
+	GENERATED_BODY()
+
+	/** Creates a FDataLinkInstance from this Task Instance */
+	FDataLinkInstance CreateDataLinkInstance() const;
+
+	/** The data link graph to execute */
+	UPROPERTY(EditAnywhere, Category="Data Link", meta=(NoBinding))
+	TObjectPtr<UDataLinkGraph> DataLinkGraph;
+
+	/**
+	 * Input data required to run the Data Link Graph.
+	 * The type of input data is determined by the Task's Data Link Graph property
+	 */
+	UPROPERTY(EditAnywhere, EditFixedSize, Category="Data Link", meta=(EditFixedOrder, NoBindingSelfOnly))
+	TArray<FDataLinkInputData> InputData;
+
+	/**
+	 * The output struct reference to write the data link results to.
+	 * NOTE: Must match the data link output struct type, or will fail to write!
+	 */
+	UPROPERTY(EditAnywhere, Category="Data Link", meta=(RefType="AnyStruct"))
+	FSceneStatePropertyReference OutputTarget;
+
+	/** The object handling the execution */
+	TSharedPtr<FDataLinkExecutor> Executor;
+};
+
+/** Runs a Data Link Graph and writes the result to the Output Target */
+USTRUCT(DisplayName="Run Data Link", Category="Data Link")
+struct FSceneStateRunDataLinkTask : public FSceneStateTask
+{
+	GENERATED_BODY()
+
+	using FInstanceDataType = FSceneStateDataLinkRequestTaskInstance;
+
+protected:
+	//~ Begin FSceneStateTask
+#if WITH_EDITOR
+	virtual const UScriptStruct* OnGetTaskInstanceType() const override;
+	virtual void OnBuildTaskInstance(UObject* InOuter, FStructView InTaskInstance) const override;
+#endif
+	virtual void OnStart(const FSceneStateExecutionContext& InContext, FStructView InTaskInstance) const override;
+	virtual void OnStop(const FSceneStateExecutionContext& InContext, FStructView InTaskInstance, ESceneStateTaskStopReason InStopReason) const override;
+	//~ End FSceneStateTask
+
+private:
+	/** Called when data link has broadcast output data */
+	static void OnOutputData(const FDataLinkExecutor& InExecutor, FConstStructView InOutputDataView, UE::SceneState::FTaskExecutionContext InTaskContext);
+
+	/** Called when a data link execution has finished */
+	static void OnFinished(const FDataLinkExecutor& InExecutor, EDataLinkExecutionResult InExecutionResult, UE::SceneState::FTaskExecutionContext InTaskContext);
+};
diff --git a/Engine/Plugins/Experimental/SceneStateDataLink/Source/SceneStateDataLinkEditor/Private/DetailsView/SceneStateDataLinkRequestTaskDetails.cpp b/Engine/Plugins/Experimental/SceneStateDataLink/Source/SceneStateDataLinkEditor/Private/DetailsView/SceneStateDataLinkRequestTaskDetails.cpp
--- a/Engine/Plugins/Experimental/SceneStateDataLink/Source/SceneStateDataLinkEditor/Private/DetailsView/SceneStateDataLinkRequestTaskDetails.cpp
+++ b/Engine/Plugins/Experimental/SceneStateDataLink/Source/SceneStateDataLinkEditor/Private/DetailsView/SceneStateDataLinkRequestTaskDetails.cpp
@@ -7,7 +7,7 @@
 #include "PropertyHandle.h"
 #include "SceneStateBlueprintEditorUtils.h"
 #include "SceneStateBlueprintEditorUtils.h"
-#include "SceneStateDataLinkRequestTask.h"
+#include "SceneStateRunDataLinkTask.h"
 
 namespace UE::SceneStateDataLink
 {
diff --git a/Engine/Plugins/Experimental/SceneStateDataLink/Source/SceneStateDataLinkEditor/Private/SceneStateDataLinkEditorModule.cpp b/Engine/Plugins/Experimental/SceneStateDataLink/Source/SceneStateDataLinkEditor/Private/SceneStateDataLinkEditorModule.cpp
--- a/Engine/Plugins/Experimental/SceneStateDataLink/Source/SceneStateDataLinkEditor/Private/SceneStateDataLinkEditorModule.cpp
+++ b/Engine/Plugins/Experimental/SceneStateDataLink/Source/SceneStateDataLinkEditor/Private/SceneStateDataLinkEditorModule.cpp
@@ -4,7 +4,7 @@
 #include "DetailsView/SceneStateDataLinkRequestTaskDetails.h"
 #include "Modules/ModuleManager.h"
 #include "PropertyEditorModule.h"
-#include "SceneStateDataLinkRequestTask.h"
+#include "SceneStateRunDataLinkTask.h"
 
 IMPLEMENT_MODULE(UE::SceneStateDataLink::FEditorModule, SceneStateDataLinkEditor)
 
