Motion Design: fixed data link scene state task destroying the data link execution on first output data.

diff --git a/Engine/Plugins/Experimental/SceneStateDataLink/Source/SceneStateDataLink/Private/SceneStateRunDataLinkTask.cpp b/Engine/Plugins/Experimental/SceneStateDataLink/Source/SceneStateDataLink/Private/SceneStateRunDataLinkTask.cpp
--- a/Engine/Plugins/Experimental/SceneStateDataLink/Source/SceneStateDataLink/Private/SceneStateRunDataLinkTask.cpp
+++ b/Engine/Plugins/Experimental/SceneStateDataLink/Source/SceneStateDataLink/Private/SceneStateRunDataLinkTask.cpp
@@ -82,8 +82,6 @@ void FSceneStateRunDataLinkTask::OnOutputData(const FDataLinkExecutor& InExecuto
 		return;
 	}
 
-	Instance->Executor.Reset();
-
 	UE::SceneState::FResolvePropertyResult Result;
 	if (!UE::SceneState::ResolveProperty(*ExecutionContext, Instance->OutputTarget, Result))
 	{
